PlainProps initial UTexture2D cook bindings without actually cooking any FTexturePlatformData

Also adds FIntPoint binding and sets the correct scope of some core types

#rnx

diff --git a/Engine/Plugins/Experimental/PlainProps/Source/Public/PlainPropsUeCoreBindings.h b/Engine/Plugins/Experimental/PlainProps/Source/Public/PlainPropsUeCoreBindings.h
--- a/Engine/Plugins/Experimental/PlainProps/Source/Public/PlainPropsUeCoreBindings.h
+++ b/Engine/Plugins/Experimental/PlainProps/Source/Public/PlainPropsUeCoreBindings.h
@@ -4,6 +4,7 @@
 
 #include "Algo/Find.h"
 #include "Containers/Set.h"
+#include "Math/IntPoint.h"
 #include "Math/Quat.h"
 #include "Math/Vector.h"
 #include "Math/Transform.h"
@@ -27,6 +28,7 @@ namespace UE::Math
 PP_REFLECT_STRUCT(, FVector, void, X, Y, Z);
 PP_REFLECT_STRUCT(, FVector4, void, X, Y, Z, W);
 PP_REFLECT_STRUCT(, FQuat, void, X, Y, Z, W);
+PP_REFLECT_STRUCT(, FIntPoint, void, X, Y);
 }
 
 template <typename T>
@@ -1139,13 +1141,19 @@ struct TVariantBinding : ICustomBinding
 	}
 };
 
+PLAINPROPS_API FEnumId DeclarePixelFormat();
+
 } // namespace PlainProps::UE
 
 namespace PlainProps
 {
 
 template <> struct TTypename<FName>					{ inline static constexpr std::string_view DeclName = "Name"; };
+template <> struct TTypename<FIntPoint>				{ inline static constexpr std::string_view DeclName = "IntPoint"; };
+template <> struct TTypename<FQuat>					{ inline static constexpr std::string_view DeclName = "Quat"; };
 template <> struct TTypename<FTransform>			{ inline static constexpr std::string_view DeclName = "Transform"; };
+template <> struct TTypename<FVector>				{ inline static constexpr std::string_view DeclName = "Vector"; };
+template <> struct TTypename<FVector4>				{ inline static constexpr std::string_view DeclName = "Vector4"; };
 template <> struct TTypename<FGuid>					{ inline static constexpr std::string_view DeclName = "Guid"; };
 template <> struct TTypename<FColor>				{ inline static constexpr std::string_view DeclName = "Color"; };
 template <> struct TTypename<FLinearColor>			{ inline static constexpr std::string_view DeclName = "LinearColor"; };
@@ -1262,6 +1270,7 @@ struct TRangeBind<TOptional<T>>
 
 template<> struct TOccupancyOf<FQuat> : FRequireAll {};
 template<> struct TOccupancyOf<FVector> : FRequireAll {};
+template<> struct TOccupancyOf<FIntPoint> : FRequireAll {};
 template<> struct TOccupancyOf<FGuid> : FRequireAll {};
 template<> struct TOccupancyOf<FColor> : FRequireAll {};
 template<> struct TOccupancyOf<FLinearColor> : FRequireAll {};
diff --git a/Engine/Plugins/Experimental/PlainPropsEngine/Source/Private/PlainPropsEngineBindings.cpp b/Engine/Plugins/Experimental/PlainPropsEngine/Source/Private/PlainPropsEngineBindings.cpp
--- a/Engine/Plugins/Experimental/PlainPropsEngine/Source/Private/PlainPropsEngineBindings.cpp
+++ b/Engine/Plugins/Experimental/PlainPropsEngine/Source/Private/PlainPropsEngineBindings.cpp
@@ -1,16 +1,151 @@
 // Copyright Epic Games, Inc. All Rights Reserved.
 
 #include "PlainPropsEngineBindings.h"
+#include "PlainPropsUObjectRuntime.h"
+#include "PlainPropsRoundtripTest.h" // EBindMode
+#include "PlainPropsBind.h"
+#include "PlainPropsLoad.h"
+#include "PlainPropsSave.h"
+#include "Algo/Compare.h"
 #include "Engine/Texture2D.h"
 #include "Logging/StructuredLog.h"
+#include "Misc/DefinePrivateMemberPtr.h"
 
 DEFINE_LOG_CATEGORY(LogPlainPropsEngine);
 
+// Temp hacks for non-intrusive prototype
+UE_DEFINE_PRIVATE_MEMBER_PTR(FIntPoint, GTexture2DImportedSize, UTexture2D, ImportedSize);
+
+////////////////////////////////////////////////////////////////////////////////////////////////
+
+template <typename T>
+struct TIsContiguousContainer<TIndirectArray<T>>
+{
+	enum { Value = true };
+};
+
+static bool DiffMip(const FTexture2DMipMap* A, const FTexture2DMipMap* B)
+{
+	unimplemented(); // todo
+	return false;
+}
+
+// In :: since Algo::Compare calls UE::Core::Private:: instead of ::UE::Core::Private::,
+// which fails to handle nested the PlainProps::UE namespace
+static bool DiffMips(const TIndirectArray<FTexture2DMipMap>& A, const TIndirectArray<FTexture2DMipMap>& B)
+{
+	return !Algo::Compare(A, B, DiffMip);
+}
+
 namespace PlainProps::UE
 {
 
+static FEnumId			IndexEnum(FScopeId Scope, FAnsiStringView Name)			{ return GUE.Names.IndexEnum({Scope, GUE.Names.MakeTypename(Name)}); }
+static FStructId		IndexStruct(FScopeId Scope, FAnsiStringView Name)		{ return GUE.Names.IndexStruct({Scope, GUE.Names.MakeTypename(Name)}); }
+static FEnumId			IndexEngineEnum(FAnsiStringView Name)					{ return IndexEnum(GUE.Scopes.Engine, Name); }
+static FStructId		IndexEngineStruct(FAnsiStringView Name)					{ return IndexStruct(GUE.Scopes.Engine, Name); }
+
+////////////////////////////////////////////////////////////////////////////////////////////////
+
+struct FTexture2DCook : ICustomBinding
+{
+	using Type = UTexture2D;
+
+	const FEnumId TextureAddressId;
+	const FDualStructId IntPointId;
+	const FDualStructId PlatformTextureDataId;
+	const FMemberId MemberIds[5];
+
+	FTexture2DCook(TPropertySpecifier<5>& Spec)
+	: TextureAddressId(IndexEngineEnum("TextureAddress"))
+	, IntPointId(IndexStruct(GUE.Scopes.Core, "IntPoint"))
+	, PlatformTextureDataId(IndexEngineStruct("PlatformTextureData"))
+	, MemberIds{GUE.Names.NameMember("AddressX"),
+				GUE.Names.NameMember("AddressY"),
+				GUE.Names.NameMember("ImportedSize"),
+				GUE.Names.NameMember("CookedFormats"),
+				GUE.Names.NameMember("CookedTextures")}
+	{
+		Spec.Members[0] = Specify<TextureAddress>(TextureAddressId);
+		Spec.Members[1] = Specify<TextureAddress>(TextureAddressId);
+		Spec.Members[2] = FMemberSpec(IntPointId);
+		Spec.Members[3] = DefaultRangeOf(Specify<EPixelFormat>(GUE.EnumIds.PixelFormat));
+		Spec.Members[4] = DefaultRangeOf(FDeclId(PlatformTextureDataId));
+	}
+	
+	inline void Save(FMemberBuilder& Dst, const UTexture2D& Src, const UTexture2D* Default, const FSaveContext& Ctx) const
+	{
+#if WITH_EDITOR // Todo: Encapsulate in a TScopedCookBinding, maybe use requires(WITH_EDTIOR)
+		FTypedRange CookedTextures;
+		FTypedRange CookedFormats = CookTexture(/* out */ CookedTextures, Src, Ctx);
+		FIntPoint ImportedSize = Src.GetImportedSize();
+
+		Dst.AddEnum(MemberIds[0], TextureAddressId, Src.AddressX.GetValue());
+		Dst.AddEnum(MemberIds[1], TextureAddressId, Src.AddressY.GetValue());
+		Dst.AddStruct(MemberIds[2], IntPointId, SaveStruct(&ImportedSize, IntPointId, Ctx));
+		Dst.AddRange(MemberIds[3], CookedFormats);
+		Dst.AddRange(MemberIds[4], CookedTextures);
+#else
+		unimplemented();
+#endif // WITH_EDITOR
+	}
+		
+	inline void Load(UTexture2D& Dst, FStructLoadView Src, ECustomLoadMethod Method) const
+	{
+		check(Method == ECustomLoadMethod::Assign);
+		FMemberLoader It(Src);
+		
+		Dst.AddressX = It.GrabLeaf().As<TextureAddress>();
+		Dst.AddressY = It.GrabLeaf().As<TextureAddress>();
+		LoadStruct(&(Dst.*GTexture2DImportedSize), It.GrabStruct());
+		FLeafRangeLoadView CookedFormats = It.GrabRange().AsLeaves();
+		FStructRangeLoadView CookedTextures = It.GrabRange().AsStructs();
+		Dst.SetPlatformData(LoadTexture(CookedFormats, CookedTextures));
+	}
+
+	static bool Diff(const UTexture2D& A, const UTexture2D& B, const FBindContext& Ctx)
+	{
+		check(WITH_EDITOR);
+		// Todo: Research how to implement this
+		return A.GetPlatformData() != B.GetPlatformData();
+	}
+
+	static bool Diff(const UTexture2D& A, const UTexture2D& B, FDiffContext& Ctx)
+	{
+		return	A.GetTextureAddressX() != B.GetTextureAddressX() ||
+				A.GetTextureAddressY() != B.GetTextureAddressY() ||
+				A.GetImportedSize() != B.GetImportedSize() ||
+				DiffTexture(*A.GetPlatformData(), *B.GetPlatformData(), Ctx);
+	}
+	
+	FTypedRange CookTexture(FTypedRange& OutCookedTextures, const UTexture2D& Src, const FSaveContext& Ctx) const
+	{
+		unimplemented(); // todo
+		return {};
+	}
+	
+	FTexturePlatformData* LoadTexture(FLeafRangeLoadView CookedFormats, FStructRangeLoadView CookedTextures) const
+	{
+		unimplemented(); // todo
+		return nullptr;
+	}
+
+	static bool DiffTexture(const FTexturePlatformData& A, const FTexturePlatformData& B, FDiffContext& Ctx)
+	{
+		return A.SizeX != B.SizeX || A.SizeY != B.SizeY || A.PackedData != B.PackedData || A.PixelFormat != B.PixelFormat ||
+			A.OptData != B.OptData || DiffMips(A.Mips, B.Mips) || !!A.VTData != !!B.VTData;
+	}
+};
+
+////////////////////////////////////////////////////////////////////////////////////////////////
+
 void CustomBindEngineTypes(EBindMode Mode)
 {
+	if (Mode == EBindMode::Runtime)
+	{
+		// Todo: Ownership / memory leak
+		new TScopedStructBinding<UTexture2D, FDeltaRuntime, FTexture2DCook>(FDualStructId(IndexEngineStruct("Texture2D")));
+	}
 }
 
 } // namespace PlainProps::UE
\ No newline at end of file
diff --git a/Engine/Plugins/Experimental/PlainPropsUObject/Source/Private/PlainPropsUObjectBindings.cpp b/Engine/Plugins/Experimental/PlainPropsUObject/Source/Private/PlainPropsUObjectBindings.cpp
--- a/Engine/Plugins/Experimental/PlainPropsUObject/Source/Private/PlainPropsUObjectBindings.cpp
+++ b/Engine/Plugins/Experimental/PlainPropsUObject/Source/Private/PlainPropsUObjectBindings.cpp
@@ -78,10 +78,6 @@ DEFINE_LOG_CATEGORY(LogPlainPropsUObject);
 namespace PlainProps::UE
 {
 
-static constexpr ESizeType DefaultRangeMax = RangeSizeOf(FDefaultAllocator::SizeType{});
-
-inline FMemberSpec DefaultRangeOf(FMemberSpec Item) { return FMemberSpec(DefaultRangeMax, Item); }
-
 ////////////////////////////////////////////////////////////////////////////////////////////////
 
 struct FDefaultStruct
@@ -256,6 +252,7 @@ const void* FDefaultStructs::Get(FBindId Id)
 FCommonScopeIds::FCommonScopeIds(TIdIndexer<FSensitiveName>& Names)
 : Core(						Names.MakeScope("/Script/Core"))
 , CoreUObject(				Names.MakeScope("/Script/CoreUObject"))
+, Engine(					Names.MakeScope("/Script/Engine"))
 {}
 
 FCommonTypenameIds::FCommonTypenameIds(TIdIndexer<FSensitiveName>& Names)
@@ -276,6 +273,10 @@ FCommonTypenameIds::FCommonTypenameIds(TIdIndexer<FSensitiveName>& Names)
 , VerseString(				Names.NameType("VerseString"))
 {}
 
+FCommonEnumIds::FCommonEnumIds(const FCommonScopeIds& Scopes, TIdIndexer<FSensitiveName>& Names)
+: PixelFormat(Names.IndexEnum({Scopes.Core, Names.MakeTypename("PixelFormat")}))
+{}
+
 FCommonStructIds::FCommonStructIds(const FCommonScopeIds& Scopes, TIdIndexer<FSensitiveName>& Names)
 : Name(						Names.IndexStruct({Scopes.Core, Names.MakeTypename("Name")}))
 , Text(						Names.IndexStruct({Scopes.Core, Names.MakeTypename("Text")}))
@@ -318,6 +319,7 @@ FGlobals::FGlobals()
 , Schemas(FDebugIds(Names))
 , Customs(FDebugIds(Names))
 , Scopes(Names)
+, EnumIds(Scopes, Names)
 , Structs(Scopes, Names)
 , Typenames(Names)
 , Members(Names)
diff --git a/Engine/Plugins/Experimental/PlainPropsUObject/Source/Public/PlainPropsUObjectRuntime.h b/Engine/Plugins/Experimental/PlainPropsUObject/Source/Public/PlainPropsUObjectRuntime.h
--- a/Engine/Plugins/Experimental/PlainPropsUObject/Source/Public/PlainPropsUObjectRuntime.h
+++ b/Engine/Plugins/Experimental/PlainPropsUObject/Source/Public/PlainPropsUObjectRuntime.h
@@ -48,6 +48,7 @@ struct FCommonScopeIds
 	
 	FScopeId					Core;
 	FScopeId					CoreUObject;
+	FScopeId					Engine; // Move out?
 };
 
 struct FCommonTypenameIds
@@ -71,28 +72,35 @@ struct FCommonTypenameIds
 	FConcreteTypenameId			VerseString;
 };
 
+struct FCommonEnumIds
+{
+	explicit FCommonEnumIds(const FCommonScopeIds& Scopes, TIdIndexer<FSensitiveName>& Names);
+	
+	FEnumId						PixelFormat;
+};
+
 struct FCommonStructIds
 {
 	explicit FCommonStructIds(const FCommonScopeIds& Scopes, TIdIndexer<FSensitiveName>& Names);
 	
-	FDualStructId			Name;
-	FDualStructId			Text;
-	FDualStructId			Guid;
-	FDualStructId			FieldPath;
-	FDualStructId			SoftObjectPath;
-	FDualStructId			ClassPtr;
-	FDualStructId			ObjectPtr;
-	FDualStructId			WeakObjectPtr;
-	FDualStructId			LazyObjectPtr;
-	FDualStructId			SoftObjectPtr;
-	FDualStructId			ScriptInterface;
-	FDualStructId			Delegate;
-	FDeclId					MulticastDelegate;
-	FBindId					MulticastInlineDelegate;
-	FDualStructId			MulticastSparseDelegate;
-	FDualStructId			VerseFunction;
-	FDualStructId			DynamicallyTypedValue;
-	FDualStructId			ReferencePropertyValue;
+	FDualStructId				Name;
+	FDualStructId				Text;
+	FDualStructId				Guid;
+	FDualStructId				FieldPath;
+	FDualStructId				SoftObjectPath;
+	FDualStructId				ClassPtr;
+	FDualStructId				ObjectPtr;
+	FDualStructId				WeakObjectPtr;
+	FDualStructId				LazyObjectPtr;
+	FDualStructId				SoftObjectPtr;
+	FDualStructId				ScriptInterface;
+	FDualStructId				Delegate;
+	FDeclId						MulticastDelegate;
+	FBindId						MulticastInlineDelegate;
+	FDualStructId				MulticastSparseDelegate;
+	FDualStructId				VerseFunction;
+	FDualStructId				DynamicallyTypedValue;
+	FDualStructId				ReferencePropertyValue;
 };
 
 struct FCommonMemberIds
@@ -122,6 +130,7 @@ struct FGlobals
 	FCustomBindingsBottom		Customs;
 	FDefaultStructs				Defaults;
 	FCommonScopeIds				Scopes;
+	FCommonEnumIds				EnumIds;
 	FCommonStructIds			Structs;
 	FCommonTypenameIds			Typenames;
 	FCommonMemberIds			Members;
@@ -160,6 +169,12 @@ struct FDeltaRuntime : FDefaultRuntime
 	static IDefaultStructs*		GetDefaults()		{ return &GUE.Defaults; }
 };
 
+//////////////////////////////////////////////////////////////////////////
+
+static constexpr ESizeType DefaultRangeMax = RangeSizeOf(FDefaultAllocator::SizeType{});
+
+inline FMemberSpec DefaultRangeOf(FMemberSpec Item) { return FMemberSpec(DefaultRangeMax, Item); }
+
 template<uint32 N>
 using TPropertySpecifier = TCustomSpecifier<FRuntimeIds, N>;
 
@@ -251,7 +266,11 @@ template<> struct TCustomBind<FReferencePropertyValue> { using Type = UE::FRefer
 template<> struct TCustomBind<FVerseFunction> { using Type = UE::FVerseFunctionBinding; };
 
 // Temporary way to tie certain types to /Script/Core scope
+template<> inline FScopeId IndexNamespaceId<UE::FRuntimeIds, TTypename<FIntPoint>>()		{ return UE::GUE.Scopes.Core; }
+template<> inline FScopeId IndexNamespaceId<UE::FRuntimeIds, TTypename<FQuat>>()			{ return UE::GUE.Scopes.Core; }
 template<> inline FScopeId IndexNamespaceId<UE::FRuntimeIds, TTypename<FTransform>>()		{ return UE::GUE.Scopes.Core; }
+template<> inline FScopeId IndexNamespaceId<UE::FRuntimeIds, TTypename<FVector>>()			{ return UE::GUE.Scopes.Core; }
+template<> inline FScopeId IndexNamespaceId<UE::FRuntimeIds, TTypename<FVector4>>()			{ return UE::GUE.Scopes.Core; }
 template<> inline FScopeId IndexNamespaceId<UE::FRuntimeIds, TTypename<FGuid>>()			{ return UE::GUE.Scopes.Core; }
 template<> inline FScopeId IndexNamespaceId<UE::FRuntimeIds, TTypename<FColor>>()			{ return UE::GUE.Scopes.Core; }
 template<> inline FScopeId IndexNamespaceId<UE::FRuntimeIds, TTypename<FLinearColor>>()		{ return UE::GUE.Scopes.Core; }
