PlainProps: commandlet: factor async loading / streaming / compilation flush into a single FlushAsyncOperations() call and use it after initial load and DropPackages

#rnx
#Synced-CL 51945425

#Synced-CL 51945425

diff --git a/Engine/Plugins/Experimental/PlainPropsEngine/Source/Private/PlainPropsCommandlets.cpp b/Engine/Plugins/Experimental/PlainPropsEngine/Source/Private/PlainPropsCommandlets.cpp
--- a/Engine/Plugins/Experimental/PlainPropsEngine/Source/Private/PlainPropsCommandlets.cpp
+++ b/Engine/Plugins/Experimental/PlainPropsEngine/Source/Private/PlainPropsCommandlets.cpp
@@ -15,10 +15,10 @@
 #include "UObject/UObjectIterator.h"
 
 #if WITH_EDITOR
+#include "AssetCompilingManager.h"
 #include "Engine/AssetManager.h"
 #include "Engine/Blueprint.h"
 #include "Kismet2/KismetEditorUtilities.h"
-#include "PackageTools.h"
 #endif
 
 // Temporary maintained filter to ignore diffs reported by RoundtripViaLinkerBatch
@@ -66,6 +66,18 @@ static PlainProps::UE::FLinkerDiffFilter LinkerDiffFilter
 	.IgnoreCastFlags = CASTCLASS_FMulticastInlineDelegateProperty | CASTCLASS_FMulticastSparseDelegateProperty,
 };
 
+static void FlushAsyncOperations(const TCHAR* Context)
+{
+	UE_LOGFMT(LogPlainPropsEngine, Display, "FlushAsyncOperations : {Context} ...", Context);
+	
+	TRACE_CPUPROFILER_EVENT_SCOPE_TEXT(Context);
+	FlushAsyncLoading();
+	(*GFlushStreamingFunc)();
+#if WITH_EDITOR
+	FAssetCompilingManager::Get().FinishAllCompilation();
+#endif
+}
+
 static FName GenerateDroppedPackageName(FName Name, int32 NewNumber = INDEX_NONE)
 {
 	if (NewNumber == INDEX_NONE)
@@ -85,14 +97,7 @@ static FName GenerateDroppedPackageName(FName Name, int32 NewNumber = INDEX_NONE
 // Code inspired by UPackageTools::ReloadPackages
 static void DropPackages(TArray<UPackage*>&& Packages, int32 NewNumber)
 {
-	{
-		TRACE_CPUPROFILER_EVENT_SCOPE(FlushLoading);
-		FlushAsyncLoading();
-		(*GFlushStreamingFunc)();
-#if WITH_EDITOR
-		UPackageTools::FlushAsyncCompilation(Packages);
-#endif
-	}
+	FlushAsyncOperations(TEXT("DropPackages"));
 
 	TArray<UObject*> AllObjects;
 	TArray<UObject*> ObjectsInPackage;
@@ -279,12 +284,7 @@ int32 UTestPlainPropsCommandlet::Main(const FString& Params)
 		}
 	}
 
-	UE_LOGFMT(LogPlainPropsEngine, Display, "FlushAsyncLoading ...");
-	{
-		TRACE_CPUPROFILER_EVENT_SCOPE(InitialFlushLoading);
-		FlushAsyncLoading();
-		(*GFlushStreamingFunc)();
-	}
+	FlushAsyncOperations(TEXT("InitialFlush"));
 
 	UE_LOGFMT(LogPlainPropsEngine, Display, "Init worlds ...");
 	InitWorlds();
