PlainProps: Fix DiffStructProperties to property ignore for non native identical structs
- Add ignore of NiagaraVariableAttributeBinding.BoundVariable

#Synced-CL 52019349

diff --git a/Engine/Plugins/Experimental/PlainPropsEngine/Source/Private/PlainPropsCommandlets.cpp b/Engine/Plugins/Experimental/PlainPropsEngine/Source/Private/PlainPropsCommandlets.cpp
--- a/Engine/Plugins/Experimental/PlainPropsEngine/Source/Private/PlainPropsCommandlets.cpp
+++ b/Engine/Plugins/Experimental/PlainPropsEngine/Source/Private/PlainPropsCommandlets.cpp
@@ -51,6 +51,7 @@ static PlainProps::UE::FLinkerDiffFilter LinkerDiffFilter
 	.IgnorePropertiesForStructs
 	{
 		TPair<FName, FName>("BaseId",							"NiagaraGraphScriptUsageInfo"		),
+		TPair<FName, FName>("BoundVariable",					"NiagaraVariableAttributeBinding"	), // Legacy editor-only migration field; CDO constructed before FNiagaraTypeDefinition::Init() so TypeDefHandle defaults to INDEX_NONE, masking the delta
 		TPair<FName, FName>("ChangeId",							"NiagaraNode"						),
 		TPair<FName, FName>("ChangeId",							"NiagaraGraph"						),
 		TPair<FName, FName>("LastBuiltTraversalDataChangeId",	"NiagaraGraph"						),
diff --git a/Engine/Plugins/Experimental/PlainPropsUObject/Source/Private/PlainPropsUObjectDiff.cpp b/Engine/Plugins/Experimental/PlainPropsUObject/Source/Private/PlainPropsUObjectDiff.cpp
--- a/Engine/Plugins/Experimental/PlainPropsUObject/Source/Private/PlainPropsUObjectDiff.cpp
+++ b/Engine/Plugins/Experimental/PlainPropsUObject/Source/Private/PlainPropsUObjectDiff.cpp
@@ -329,6 +329,12 @@ bool DiffStructProperties(const void* A, const void* B, const FStructProperty* P
 	}
 	if (DiffStructs(A, B, ScriptStruct, Ctx))
 	{
+		if (IgnoreDiff(A, B, Property, Ctx))
+		{
+			Ctx.Diffs.Reset();
+			return false;
+		}
+
 		Ctx.Diffs.Emplace(EDiffObjectNodeType::Property, A, B, Property);
 		return true;
 	}
