PlainProps: Add FDynamicIds to FGlobals for UEnum->FEnumId and UStruct->UStructId lookups.
Add TScopedUStructBinding and TScopedUClassBinding helpers derived from TScopedStructBinding.
Adjust custom binding implementations to use GUE.DynamicIds.Get[Struct|Enum] and nuke some obsolete TTypename and GetOtuerScope template specializations.

#rnx

#Synced-CL 52129858

diff --git a/Engine/Plugins/Experimental/PlainPropsEngine/Source/Private/PlainPropsAnimationBindings.cpp b/Engine/Plugins/Experimental/PlainPropsEngine/Source/Private/PlainPropsAnimationBindings.cpp
--- a/Engine/Plugins/Experimental/PlainPropsEngine/Source/Private/PlainPropsAnimationBindings.cpp
+++ b/Engine/Plugins/Experimental/PlainPropsEngine/Source/Private/PlainPropsAnimationBindings.cpp
@@ -83,46 +83,6 @@ struct FAnimationAttributeIdentifierBinding : ICustomBinding
 	}
 };
 
-struct FAttributeIds
-{
-	static constexpr std::initializer_list<const char*> AnsiNames =
-	{
-		"IntegerAnimationAttribute",
-		"FloatAnimationAttribute",
-		"VectorAnimationAttribute",
-		"QuaternionAnimationAttribute",
-		"TransformAnimationAttribute",
-		"StringAnimationAttribute",
-		"NonBlendableIntegerAnimationAttribute",
-		"NonBlendableFloatAnimationAttribute",
-		"NonBlendableVectorAnimationAttribute",
-		"NonBlendableQuaternionAnimationAttribute",
-		"NonBlendableTransformAnimationAttribute",
-	};
-
-	FName Names[AnsiNames.size()];
-	FBindId BindIds[AnsiNames.size()];
-
-	template<class Ids>
-	FAttributeIds(TInit<Ids>)
-	{
-		FName* NamePtr = &Names[0];
-		FBindId* IdPtr = &BindIds[0];
-		for (const char* Name : AnsiNames)
-		{
-			*NamePtr++ = FName(Name);
-			*IdPtr++ = FBindId(Ids::IndexStruct(GetAnimationScope(), Name));
-		}
-	}
-
-	inline const FBindId GetBindId(UScriptStruct* Struct) const
-	{
-		int32 I = Algo::IndexOf(Names, Struct->GetFName());
-		checkf(I != INDEX_NONE, TEXT("Unknown attribute identifier '%s', it has not been indexed."), *Struct->GetFullName());
-		return BindIds[I];
-	}
-};
-
 struct FAttributeCurveBinding : ICustomBinding
 {
 	using Type = FAttributeCurve;
@@ -130,12 +90,10 @@ struct FAttributeCurveBinding : ICustomBinding
 
 	const FDualStructId AttributeKeyId;
 	const FMemberId MemberIds[3];
-	FAttributeIds AttributeIds;
 	
 	FAttributeCurveBinding(TPropertySpecifier<3>& Spec)
-	: AttributeKeyId(GetStructDualId<Ids, FAttributeKey>())
+	: AttributeKeyId(GUE.DynamicIds.GetStruct(FAttributeKey::StaticStruct()))
 	, MemberIds{Ids::IndexMember("Keys"), Ids::IndexMember("ScriptStructPath"), Ids::IndexMember("DynamicKeys")}
-	, AttributeIds(TInit<Ids>{})
 	{
 		Spec.Members[0] = DefaultRangeOf(FDeclId(AttributeKeyId));
 		Spec.Members[1] = FDeclId(GUE.Structs.SoftObjectPath);
@@ -170,7 +128,7 @@ struct FAttributeCurveBinding : ICustomBinding
 
 		Dst.AddRange(MemberIds[0], SaveKeysStatic(AttributeKeyId, Src.*_Keys, Ctx));
 		Dst.AddStruct(MemberIds[1], GUE.Structs.SoftObjectPath, SaveStruct(&(Src.*_ScriptStructPath), GUE.Structs.SoftObjectPath, Ctx));
-		FBindId DynamicKeyId = AttributeIds.GetBindId(Src.*_ScriptStruct);
+		FBindId DynamicKeyId = GUE.DynamicIds.GetStruct(Src.*_ScriptStruct);
 		Dst.AddRange(MemberIds[2], SaveKeysDynamic(DynamicKeyId, Src.*_Keys, Ctx));
 	}
 
@@ -263,10 +221,10 @@ void CustomBindAnimationTypes(EBindMode Mode)
 {
 	using namespace Animation;
 	// Todo: Ownership / memory leak
-	new TScopedStructBinding<FAnimationAttributeIdentifier, FDefaultRuntime, Animation::FAnimationAttributeIdentifierBinding>();
-	new TScopedStructBinding<FAttributeCurve, FDefaultRuntime, Animation::FAttributeCurveBinding>();
-}
+	new TScopedDefaultUStructBinding<FAnimationAttributeIdentifier, FAnimationAttributeIdentifierBinding>();
+	new TScopedDefaultUStructBinding<FAttributeCurve, FAttributeCurveBinding>();
 }
+} // namespace PlainProps::UE
 
 ////////////////////////////////////////////////////////////////////////////////////////////////
 
@@ -275,12 +233,4 @@ namespace PlainProps
 template<> struct TOccupancyOf<FAnimationAttributeIdentifier> : FRequireAll {};
 template<> struct TOccupancyOf<FAttributeKey> : FRequireAll {};
 template<> struct TOccupancyOf<FAttributeCurve> : FRequireAll {};
-
-// Enable usage of TScopedStructBinding default constructor as well as GetStructXXXXId<Ids,Struct> and GetEnumId<Ids,Enum> lookups in custom bindings.
-template<> struct TTypename<FAnimationAttributeIdentifier>		{ inline static constexpr std::string_view DeclName = "AnimationAttributeIdentifier"; }; 
-template<> struct TTypename<FAttributeKey>						{ inline static constexpr std::string_view DeclName = "AttributeKey"; }; 
-template<> struct TTypename<FAttributeCurve>					{ inline static constexpr std::string_view DeclName = "AttributeCurve"; }; 
-template<> inline FScopeId GetOuterScope<UE::FRuntimeIds, TTypename<FAnimationAttributeIdentifier>>()		{ return UE::Animation::GetAnimationScope(); }
-template<> inline FScopeId GetOuterScope<UE::FRuntimeIds, TTypename<FAttributeKey>>()						{ return UE::Animation::GetAnimationScope(); }
-template<> inline FScopeId GetOuterScope<UE::FRuntimeIds, TTypename<FAttributeCurve>>()						{ return UE::Animation::GetAnimationScope(); }
-}
+}  // namespace PlainProps
diff --git a/Engine/Plugins/Experimental/PlainPropsEngine/Source/Private/PlainPropsEngineBindings.cpp b/Engine/Plugins/Experimental/PlainPropsEngine/Source/Private/PlainPropsEngineBindings.cpp
--- a/Engine/Plugins/Experimental/PlainPropsEngine/Source/Private/PlainPropsEngineBindings.cpp
+++ b/Engine/Plugins/Experimental/PlainPropsEngine/Source/Private/PlainPropsEngineBindings.cpp
@@ -141,12 +141,6 @@ struct FTexture2DCook : ICustomBinding
 
 ////////////////////////////////////////////////////////////////////////////////////////////////
 
-static FScopeId GetAbilitiesScope()
-{
-	static FScopeId Scope(UE::FRuntimeIds::IndexScope("/Script/GameplayAbilities"));
-	return Scope;
-}
-
 struct FGameplayEffectVersionBinding : ICustomBinding
 {
 	using Type = FGameplayEffectVersion;
@@ -156,7 +150,7 @@ struct FGameplayEffectVersionBinding : ICustomBinding
 	const FMemberId MemberIds[1];
 	
 	FGameplayEffectVersionBinding(TPropertySpecifier<1>& Spec)
-	: GameplayEffectVersionId(Ids::IndexEnum(GetAbilitiesScope(), "EGameplayEffectVersion"))
+	: GameplayEffectVersionId(GUE.DynamicIds.GetEnum(StaticEnum<EGameplayEffectVersion>()))
 	, MemberIds{Ids::IndexMember("CurrentVersion")}
 	{
 		Spec.Members[0] = Specify<EGameplayEffectVersion>(GameplayEffectVersionId);
@@ -189,10 +183,11 @@ void CustomBindEngineTypes(EBindMode Mode)
 	{ DbgVis::FLoadStructPlan _; DbgVis::FIdVisualizer{}.KeepDebugInfo(&_); }
 
 	// Todo: Ownership / memory leak
-	new TScopedStructBinding<FGameplayEffectVersion, FDefaultRuntime, FGameplayEffectVersionBinding>();
+	new TScopedDefaultUStructBinding<FGameplayEffectVersion, FGameplayEffectVersionBinding>();
+
 	if (Mode == EBindMode::Runtime)
 	{
-		new TScopedStructBinding<UTexture2D, FDeltaRuntime, FTexture2DCook>(FDualStructId(IndexEngineStruct("Texture2D")));
+		new TScopedUClassBinding<UTexture2D, FDeltaRuntime, FTexture2DCook>();
 	}
 }
 
@@ -201,6 +196,4 @@ void CustomBindEngineTypes(EBindMode Mode)
 namespace PlainProps
 {
 template<> struct TOccupancyOf<FGameplayEffectVersion> : FRequireAll {};
-template<> struct TTypename<FGameplayEffectVersion> { inline static constexpr std::string_view DeclName = "GameplayEffectVersion"; }; 
-template<> inline FScopeId GetOuterScope<UE::FRuntimeIds, TTypename<FGameplayEffectVersion>>() { return UE::GetAbilitiesScope(); }
-}
+} // namespace PlainProps
diff --git a/Engine/Plugins/Experimental/PlainPropsEngine/Source/Private/PlainPropsNiagaraBindings.cpp b/Engine/Plugins/Experimental/PlainPropsEngine/Source/Private/PlainPropsNiagaraBindings.cpp
--- a/Engine/Plugins/Experimental/PlainPropsEngine/Source/Private/PlainPropsNiagaraBindings.cpp
+++ b/Engine/Plugins/Experimental/PlainPropsEngine/Source/Private/PlainPropsNiagaraBindings.cpp
@@ -27,12 +27,6 @@ PP_REFLECT_ENUM(FNiagaraTypeDefinition, FTypeFlags, TF_None, TF_Static);
 namespace PlainProps::UE::Niagara
 {
 
-static FScopeId GetNiagaraScope()
-{
-	static FScopeId Scope(UE::FRuntimeIds::IndexScope("/Script/Niagara"));
-	return Scope;
-}
-
 // Custom bind FNiagaraTypeDefinition to modify the property values on Save and Load,
 // could use a schema binding for actual serialization.
 struct FNiagaraTypeDefinitionBinding : ICustomBinding
@@ -93,7 +87,7 @@ struct FNiagaraTypeDefinitionHandleBinding : ICustomBinding
 	const FMemberId MemberIds[1];
 	
 	FNiagaraTypeDefinitionHandleBinding(TPropertySpecifier<1>& Spec)
-	: NiagaraTypeDefinitionId(GetStructDualId<Ids, FNiagaraTypeDefinition>())
+	: NiagaraTypeDefinitionId(GUE.DynamicIds.GetStruct(FNiagaraTypeDefinition::StaticStruct()))
 	, MemberIds{Ids::IndexMember("TypeDef")}
 	{
 		Spec.Members[0] = FDeclId(NiagaraTypeDefinitionId);
@@ -125,12 +119,12 @@ namespace PlainProps::UE
 void CustomBindNiagaraTypes(EBindMode Mode)
 {
 	// Todo: Ownership / memory leak
-	new TScopedEnumDeclaration<FNiagaraTypeDefinition::FUnderlyingType, EEnumMode::Flat, FDefaultRuntime>();
-	new TScopedEnumDeclaration<FNiagaraTypeDefinition::FTypeFlags, EEnumMode::Flag, FDefaultRuntime>();
-	new TScopedStructBinding<FNiagaraTypeDefinition, FDefaultRuntime, Niagara::FNiagaraTypeDefinitionBinding>();
-	new TScopedStructBinding<FNiagaraTypeDefinitionHandle, FDefaultRuntime, Niagara::FNiagaraTypeDefinitionHandleBinding>();
+	new TScopedDefaultEnumDeclaration<FNiagaraTypeDefinition::FUnderlyingType, EEnumMode::Flat>();	// global scope
+	new TScopedDefaultEnumDeclaration<FNiagaraTypeDefinition::FTypeFlags, EEnumMode::Flag>();		// global scope
+	new TScopedDefaultUStructBinding<FNiagaraTypeDefinition, Niagara::FNiagaraTypeDefinitionBinding>();
+	new TScopedDefaultUStructBinding<FNiagaraTypeDefinitionHandle, Niagara::FNiagaraTypeDefinitionHandleBinding>();
 }
-} 
+} // namespace PlainProps::UE
 
 ////////////////////////////////////////////////////////////////////////////////////////////////
 namespace PlainProps
@@ -138,11 +132,7 @@ namespace PlainProps
 template<> struct TOccupancyOf<FNiagaraTypeDefinition> : FRequireAll {};
 template<> struct TOccupancyOf<FNiagaraTypeDefinitionHandle> : FRequireAll {};
 
-// Enable usage of TScopedStructBinding default constructor as well as GetStructXXXXId<Ids,Struct> and GetEnumId<Ids,Enum> lookups in custom bindings.
-template<> struct TTypename<FNiagaraTypeDefinition>			{ inline static constexpr std::string_view DeclName = "NiagaraTypeDefinition"; }; 
-template<> struct TTypename<FNiagaraTypeDefinitionHandle>	{ inline static constexpr std::string_view DeclName = "NiagaraTypeDefinitionHandle"; }; 
-template<> inline FScopeId GetOuterScope<UE::FRuntimeIds, TTypename<FNiagaraTypeDefinition>>()					{ return UE::Niagara::GetNiagaraScope(); }
-template<> inline FScopeId GetOuterScope<UE::FRuntimeIds, TTypename<FNiagaraTypeDefinitionHandle>>()			{ return UE::Niagara::GetNiagaraScope(); }
-template<> inline FScopeId GetOuterScope<UE::FRuntimeIds, TTypename<FNiagaraTypeDefinition::FUnderlyingType>>()	{ return UE::Niagara::GetNiagaraScope(); }
-template<> inline FScopeId GetOuterScope<UE::FRuntimeIds, TTypename<FNiagaraTypeDefinition::FTypeFlags>>()		{ return UE::Niagara::GetNiagaraScope(); }
-}
+// todo: Optionally add a fake "/Script/Niagara" scope to these non-reflected class internal native enums?
+// template<> inline FScopeId GetOuterScope<UE::FRuntimeIds, TTypename<FNiagaraTypeDefinition::FUnderlyingType>>()	{ return UE::Niagara::GetNiagaraScope(); }
+// template<> inline FScopeId GetOuterScope<UE::FRuntimeIds, TTypename<FNiagaraTypeDefinition::FTypeFlags>>()		{ return UE::Niagara::GetNiagaraScope(); }
+} // namespace PlainProps
diff --git a/Engine/Plugins/Experimental/PlainPropsUObject/Source/Private/PlainPropsUObjectBindings.cpp b/Engine/Plugins/Experimental/PlainPropsUObject/Source/Private/PlainPropsUObjectBindings.cpp
--- a/Engine/Plugins/Experimental/PlainPropsUObject/Source/Private/PlainPropsUObjectBindings.cpp
+++ b/Engine/Plugins/Experimental/PlainPropsUObject/Source/Private/PlainPropsUObjectBindings.cpp
@@ -1117,6 +1117,7 @@ static FEnumId DeclareEnum(UEnum* Enum)
 #endif
 
 	GUE.Enums.Declare(Id, Type, Mode, Width, Enumerators, EEnumAliases::Strip);
+	GUE.DynamicIds.AddEnum(Enum, Id);
 	return Id;
 }
 
@@ -4410,6 +4411,7 @@ void BindStruct(FBindId Id, const UStruct* Struct)
 	BindMembers(/* out */ Binder, Struct);
 
 	GUE.Schemas.BindStruct(Id, Binder.GetMembers(), Binder.Declare());
+	GUE.DynamicIds.AddStruct(Struct, Id);
 	
 	// Don't bind CDO defaults, object defaults are passed in from top and objects aren't owned by containers 
 	if (Struct->HasAnyCastFlags(CASTCLASS_UScriptStruct))
diff --git a/Engine/Plugins/Experimental/PlainPropsUObject/Source/Public/PlainPropsUObjectRuntime.h b/Engine/Plugins/Experimental/PlainPropsUObject/Source/Public/PlainPropsUObjectRuntime.h
--- a/Engine/Plugins/Experimental/PlainPropsUObject/Source/Public/PlainPropsUObjectRuntime.h
+++ b/Engine/Plugins/Experimental/PlainPropsUObject/Source/Public/PlainPropsUObjectRuntime.h
@@ -124,6 +124,21 @@ struct FCommonMemberIds
 	FMemberId					Owner;
 };
 
+// temporary non-intrusive mapping from UEnum->FEnumId and UStruct->FBindId,
+// these ids will be owned directly by UEnum and UStruct in a future integration
+struct FDynamicIds
+{
+	void			AddEnum(const UEnum* Enum, FEnumId Id)			{ Enums.Add(Enum, Id); }
+	void			AddStruct(const UStruct* Struct, FBindId Id)	{ Structs.Add(Struct, FDualStructId(Id)); }
+
+	FEnumId			GetEnum(const UEnum* Enum) const				{ return Enums.FindChecked(Enum); }
+	FDualStructId	GetStruct(const UStruct* Struct) const			{ return Structs.FindChecked(Struct); }
+
+private:
+	TMap<const UEnum*, FEnumId> Enums;
+	TMap<const UStruct*, FDualStructId> Structs;
+};
+
 struct FGlobals
 {
 	FGlobals();
@@ -138,6 +153,7 @@ struct FGlobals
 	FCommonStructIds			Structs;
 	FCommonTypenameIds			Typenames;
 	FCommonMemberIds			Members;
+	FDynamicIds					DynamicIds;
 	FNumeralGenerator			Numerals;
 	//Upgrade::FHistory			Upgrades;
 	FDebugIds					Debug;
@@ -186,6 +202,36 @@ inline FMemberSpec DefaultRangeOf(FMemberSpec Item) { return FMemberSpec(Default
 template<uint32 N>
 using TPropertySpecifier = TCustomSpecifier<FRuntimeIds, N>;
 
+// Default runtime scope to declare a ctti reflected native enum
+template<typename Enum, EEnumMode Mode>
+using TScopedDefaultEnumDeclaration = TScopedEnumDeclaration<Enum, Mode, FDefaultRuntime>;
+
+// Runtime scope to declare and custom bind an override for a schema bound UStruct
+template<class T, class Runtime, typename CustomBinding>
+struct TScopedUStructBinding : TScopedStructBinding<T, Runtime, CustomBinding>
+{
+	TScopedUStructBinding()
+	: TScopedStructBinding<T, Runtime, CustomBinding>(GUE.DynamicIds.GetStruct(T::StaticStruct()))
+	{}
+};
+
+// Runtime scope to declare and custom bind an override for a schema bound UClass
+template<class T, class Runtime, typename CustomBinding>
+struct TScopedUClassBinding : TScopedStructBinding<T, Runtime, CustomBinding>
+{
+	TScopedUClassBinding()
+	: TScopedStructBinding<T, Runtime, CustomBinding>(GUE.DynamicIds.GetStruct(T::StaticClass()))
+	{}
+};
+
+// Default runtime scope to declare and custom bind an override for a schema bound UStruct
+template<class T, typename CustomBinding>
+using TScopedDefaultUStructBinding = TScopedUStructBinding<T, FDefaultRuntime, CustomBinding>;
+
+// Default runtime scope to declare and custom bind an override for a schema bound UClass
+template<class T, typename CustomBinding>
+using TScopedDefaultUClassBinding = TScopedUClassBinding<T, FDefaultRuntime, CustomBinding>;
+
 //////////////////////////////////////////////////////////////////////////
 
 struct FFieldPathBinding : ICustomBinding
