PlainProps: remove filtering of FNiagaraVariableAttributeBinding::BoundVariable (was fixed in Niagara CL 52302921)

#rnx
#Synced-CL 52296502

diff --git a/Engine/Plugins/Experimental/PlainPropsEngine/Source/Private/PlainPropsCommandlets.cpp b/Engine/Plugins/Experimental/PlainPropsEngine/Source/Private/PlainPropsCommandlets.cpp
--- a/Engine/Plugins/Experimental/PlainPropsEngine/Source/Private/PlainPropsCommandlets.cpp
+++ b/Engine/Plugins/Experimental/PlainPropsEngine/Source/Private/PlainPropsCommandlets.cpp
@@ -51,7 +51,6 @@ static PlainProps::UE::FLinkerDiffFilter LinkerDiffFilter
 	.IgnorePropertiesForStructs
 	{
 		TPair<FName, FName>("BaseId",							"NiagaraGraphScriptUsageInfo"		),
-		TPair<FName, FName>("BoundVariable",					"NiagaraVariableAttributeBinding"	), // Legacy editor-only migration field; CDO constructed before FNiagaraTypeDefinition::Init() so TypeDefHandle defaults to INDEX_NONE, masking the delta
 		TPair<FName, FName>("ChangeId",							"NiagaraNode"						),
 		TPair<FName, FName>("ChangeId",							"NiagaraGraph"						),
 		TPair<FName, FName>("LastBuiltTraversalDataChangeId",	"NiagaraGraph"						),
